- New icons for modern xcode mode

diff --git a/Engine/Build/Mac/Resources/Assets.xcassets/AppIcon.appiconset/Contents.json b/Engine/Build/Mac/Resources/Assets.xcassets/AppIcon.appiconset/Contents.json
--- a/Engine/Build/Mac/Resources/Assets.xcassets/AppIcon.appiconset/Contents.json
+++ b/Engine/Build/Mac/Resources/Assets.xcassets/AppIcon.appiconset/Contents.json
@@ -7,7 +7,7 @@
       "size" : "16x16"
     },
     {
-      "filename" : "icon_32x32 1.png",
+      "filename" : "[REDACTED_EMAIL]",
       "idiom" : "mac",
       "scale" : "2x",
       "size" : "16x16"
@@ -19,6 +19,7 @@
       "size" : "32x32"
     },
     {
+      "filename" : "[REDACTED_EMAIL]",
       "idiom" : "mac",
       "scale" : "2x",
       "size" : "32x32"
@@ -30,7 +31,7 @@
       "size" : "128x128"
     },
     {
-      "filename" : "icon_256x256 1.png",
+      "filename" : "[REDACTED_EMAIL]",
       "idiom" : "mac",
       "scale" : "2x",
       "size" : "128x128"
@@ -42,7 +43,7 @@
       "size" : "256x256"
     },
     {
-      "filename" : "icon_512x512 1.png",
+      "filename" : "[REDACTED_EMAIL]",
       "idiom" : "mac",
       "scale" : "2x",
       "size" : "256x256"
@@ -54,6 +55,7 @@
       "size" : "512x512"
     },
     {
+      "filename" : "[REDACTED_EMAIL]",
       "idiom" : "mac",
       "scale" : "2x",
       "size" : "512x512"
