- Testing modern xcode buildscript

diff --git a/Engine/Build/Mac/ModernBuildGraph.xml b/Engine/Build/Mac/ModernBuildGraph.xml
new file mode 100644
--- /dev/null
+++ b/Engine/Build/Mac/ModernBuildGraph.xml
@@ -0,0 +1,47 @@
+<?xml version='1.0'?>
+<BuildGraph xmlns="http://www.epicgames.com/BuildGraph" xmlns:xsi="http://www.w3.org/2001/XMLSchema-instance" xsi:schemaLocation="http://www.epicgames.com/BuildGraph ../Graph/Schema.xsd">
+
+	<Option Name="TestPrograms" DefaultValue="true" Description="Test building/packaging Programs" />
+	<Option Name="TestContentOnly" DefaultValue="true" Description="Test building/packaging Content Only projects" />
+	<Option Name="TestCodeBased" DefaultValue="true" Description="Test building/packaging Code Base projects" />
+	<Option Name="TestHybrid" DefaultValue="true" Description="Test building/packaging Hybrid Projects" />
+
+	<Property Name="UBTArgs" Value="-modernxcode" />
+	<Property Name="BCRArgs" Value="-platform=Mac -build -cook -stage -pak -package -modernxcode" />
+
+
+	<Agent Name="Modern Xcode Programs" Type="Mac">
+		<Node Name="Build Programs">
+			<Compile Target="CrashReportClient" Platform="Mac" Configuration="Development" Arguments="$(UBTArgs)" />
+			<Command Name="BuildCookRun" Arguments="$(BCRArgs) -project=SlateUGS -skipcook" />
+		</Node>
+	</Agent>
+
+	<Agent Name="Modern Xcode ContentOnly" Type="Mac">
+		<Node Name="Build Content Only">
+			<Compile Target="UnrealGame" Platform="Mac" Configuration="Development" Arguments="$(UBTArgs)" />
+			<Command Name="BuildCookRun" Arguments="$(BCRArgs) -project=TP_PuzzleBP" />
+		</Node>
+	</Agent>
+
+	<Agent Name="Modern Xcode CodeBased" Type="Mac">
+		<Node Name="Build Code Based">
+			<Compile Target="QAGame" Platform="Mac" Configuration="Development" Arguments="$(UBTArgs)" />
+			<Command Name="BuildCookRun" Arguments="$(BCRArgs) -project=QAGame" />
+		</Node>
+	</Agent>
+
+	<Agent Name="Modern Xcode Hybrid" Type="Mac">
+		<Node Name="Build Hybrid">
+			<Compile Target="Cropout" Platform="Mac" Configuration="Development" Arguments="$(UBTArgs)" />
+			<Command Name="BuildCookRun" Arguments="$(BCRArgs) -project=Cropout" />
+		</Node>
+	</Agent>
+	
+	<Property Name="AllBuilds" Value="Build Programs" If="$(TestPrograms)" />
+	<Property Name="AllBuilds" Value="$(AllBuilds);Build Content Only" If="$(TestContentOnly)" />
+	<Property Name="AllBuilds" Value="$(AllBuilds);Build Code Based" If="$(TestCodeBased)" />
+	<Property Name="AllBuilds" Value="$(AllBuilds);Build Hybrid" If="$(TestHybrid)" />
+	<Aggregate Name="Test Modern Xcode" Requires="$(AllBuilds)"/>
+
+</BuildGraph>
